Cloth - Added a new GetClothingAssetFactories method to allow for multiple factories and added a CVAR to decide to keep the legacy factory the default factory.

#rnx

diff --git a/Engine/Source/Editor/ClothingSystemEditorInterface/Private/ClothingSystemEditorInterfaceModule.cpp b/Engine/Source/Editor/ClothingSystemEditorInterface/Private/ClothingSystemEditorInterfaceModule.cpp
--- a/Engine/Source/Editor/ClothingSystemEditorInterface/Private/ClothingSystemEditorInterfaceModule.cpp
+++ b/Engine/Source/Editor/ClothingSystemEditorInterface/Private/ClothingSystemEditorInterfaceModule.cpp
@@ -5,6 +5,7 @@
 #include "ClothingAssetFactoryInterface.h"
 #include "Containers/Array.h"
 #include "Features/IModularFeatures.h"
+#include "HAL/IConsoleManager.h"
 #include "HAL/Platform.h"
 #include "Modules/ModuleManager.h"
 #include "SimulationEditorExtender.h"
@@ -12,6 +13,19 @@
 
 IMPLEMENT_MODULE(FClothingSystemEditorInterfaceModule, ClothingSystemEditorInterface);
 
+
+namespace UE::Private
+{
+	static FString DefaultClothingAssetFactory = TEXT("ClothingAssetFactory");  // UClothingAssetFactory is the default provider 
+
+	static TAutoConsoleVariable<FString> CVarDefaultClothingAssetFactoryClass(
+		TEXT("p.Cloth.DefaultClothingAssetFactoryClass"),
+		DefaultClothingAssetFactory,
+		TEXT("The class name of the default clothing asset factory."),
+		ECVF_Default);
+}
+
+
 const FName FClothingSystemEditorInterfaceModule::ExtenderFeatureName(TEXT("ClothingSimulationEditorExtender"));
 
 FClothingSystemEditorInterfaceModule::FClothingSystemEditorInterfaceModule()
@@ -33,14 +47,39 @@ UClothingAssetFactoryBase* FClothingSystemEditorInterfaceModule::GetClothingAsse
 {
 	TArray<IClothingAssetFactoryProvider*> Providers = IModularFeatures::Get().GetModularFeatureImplementations<IClothingAssetFactoryProvider>(IClothingAssetFactoryProvider::FeatureName);
 
-	// If we have providers, take the last one to register so clothing asst factories can be overridden
-	if(Providers.Num() > 0)
+	UClothingAssetFactoryBase* ClothingAssetFactory = nullptr;
+	for (IClothingAssetFactoryProvider* const Provider : Providers)
 	{
-		IClothingAssetFactoryProvider* Provider = Providers.Last();
-		return Provider->GetFactory();
+		if (Provider)
+		{
+			if (UClothingAssetFactoryBase* const Factory = Provider->GetFactory())
+			{
+				if (Factory->GetClass()->GetName() == UE::Private::CVarDefaultClothingAssetFactoryClass.GetValueOnAnyThread())
+				{
+					ClothingAssetFactory = Factory;
+				}
+			}
+		}
 	}
+	return ClothingAssetFactory;
+}
 
-	return nullptr;
+TArray<UClothingAssetFactoryBase*> FClothingSystemEditorInterfaceModule::GetClothingAssetFactories()
+{
+	const TArray<IClothingAssetFactoryProvider*> Providers = IModularFeatures::Get().GetModularFeatureImplementations<IClothingAssetFactoryProvider>(IClothingAssetFactoryProvider::FeatureName);
+
+	TArray<UClothingAssetFactoryBase*> ClothingAssetFactories;
+	for (IClothingAssetFactoryProvider* const Provider : Providers)
+	{
+		if (Provider)
+		{
+			if (UClothingAssetFactoryBase* const Factory = Provider->GetFactory())
+			{
+				ClothingAssetFactories.Emplace(Factory);
+			}
+		}
+	}
+	return ClothingAssetFactories;
 }
 
 ISimulationEditorExtender* FClothingSystemEditorInterfaceModule::GetSimulationEditorExtender(FName InSimulationClassName)
diff --git a/Engine/Source/Editor/ClothingSystemEditorInterface/Public/ClothingSystemEditorInterfaceModule.h b/Engine/Source/Editor/ClothingSystemEditorInterface/Public/ClothingSystemEditorInterfaceModule.h
--- a/Engine/Source/Editor/ClothingSystemEditorInterface/Public/ClothingSystemEditorInterfaceModule.h
+++ b/Engine/Source/Editor/ClothingSystemEditorInterface/Public/ClothingSystemEditorInterfaceModule.h
@@ -23,6 +23,7 @@ public:
 	UE_API virtual void ShutdownModule() override;
 
 	UE_API UClothingAssetFactoryBase* GetClothingAssetFactory();
+	UE_API TArray<UClothingAssetFactoryBase*> GetClothingAssetFactories();
 	UE_API ISimulationEditorExtender* GetSimulationEditorExtender(FName InSimulationFactoryClassName);
 
 private:
