Cloth:
- Deprecated CreateFromExistingCloth, and replaced it by CreateFromSourceAsset.
- Added GetSupportedSourceAssetType.
- Deprecated unused legacy import methods.

#rnx

diff --git a/Engine/Source/Editor/ClothPainter/Private/SClothAssetSelector.cpp b/Engine/Source/Editor/ClothPainter/Private/SClothAssetSelector.cpp
--- a/Engine/Source/Editor/ClothPainter/Private/SClothAssetSelector.cpp
+++ b/Engine/Source/Editor/ClothPainter/Private/SClothAssetSelector.cpp
@@ -1205,21 +1205,34 @@ void SClothAssetSelector::PostUndo(bool bSuccess)
 
 void SClothAssetSelector::OnCopyClothingAssetSelected(const FAssetData& AssetData)
 {
-	USkeletalMesh* SourceSkelMesh = Cast<USkeletalMesh>(AssetData.GetAsset());
-
-	if (Mesh && SourceSkelMesh && Mesh != SourceSkelMesh)
+	if (const UObject* const SourceAsset = AssetData.GetAsset())
 	{
-		FScopedTransaction Transaction(LOCTEXT("CopiedClothingAssetsFromSkelMesh", "Copied clothing assets from another SkelMesh"));
-		Mesh->Modify();
-		FClothingSystemEditorInterfaceModule& ClothingEditorModule = FModuleManager::LoadModuleChecked<FClothingSystemEditorInterfaceModule>("ClothingSystemEditorInterface");
-		UClothingAssetFactoryBase* AssetFactory = ClothingEditorModule.GetClothingAssetFactory();
-
-		for (UClothingAssetBase* ClothingAsset : SourceSkelMesh->GetMeshClothingAssets())
+		if (Mesh && Mesh != Cast<USkeletalMesh>(SourceAsset))
 		{
-			UClothingAssetCommon* NewAsset = Cast<UClothingAssetCommon>(AssetFactory->CreateFromExistingCloth(Mesh, SourceSkelMesh, ClothingAsset));
-			Mesh->AddClothingAsset(NewAsset);
+			TOptional<FScopedTransaction> Transaction;
+
+			FClothingSystemEditorInterfaceModule& ClothingEditorModule = FModuleManager::LoadModuleChecked<FClothingSystemEditorInterfaceModule>("ClothingSystemEditorInterface");
+			for (const UClothingAssetFactoryBase* const ClothingAssetFactory : ClothingEditorModule.GetClothingAssetFactories())
+			{
+				if (SourceAsset->GetClass()->IsChildOf(ClothingAssetFactory->GetSupportedSourceAssetType()))
+				{
+					if (!Transaction.IsSet())
+					{
+						Transaction.Emplace(LOCTEXT("CreatingClothingAsset", "Creating clothing asset from another asset."));
+						Mesh->Modify();
+					}
+
+					for (UClothingAssetBase* const NewAsset : ClothingAssetFactory->CreateFromSourceAsset(Mesh, SourceAsset))
+					{
+						Mesh->AddClothingAsset(NewAsset);
+					}
+				}
+			}
+			if (Transaction.IsSet())
+			{
+				OnRefresh();
+			}
 		}
-		OnRefresh();
 	}
 	FSlateApplication::Get().DismissAllMenus();
 }
diff --git a/Engine/Source/Editor/ClothingSystemEditor/Private/ClothingAssetFactory.cpp b/Engine/Source/Editor/ClothingSystemEditor/Private/ClothingAssetFactory.cpp
--- a/Engine/Source/Editor/ClothingSystemEditor/Private/ClothingAssetFactory.cpp
+++ b/Engine/Source/Editor/ClothingSystemEditor/Private/ClothingAssetFactory.cpp
@@ -101,19 +101,31 @@ UClothingAssetFactory::UClothingAssetFactory(const FObjectInitializer& ObjectIni
 {
 }
 
-UClothingAssetBase* UClothingAssetFactory::Import
-(
-	const FString& Filename, 
-	USkeletalMesh* TargetMesh,
-	FName InName
-)
+const UClass* UClothingAssetFactory::GetSupportedSourceAssetType() const
 {
-	return nullptr;
+	return USkeletalMesh::StaticClass();
 }
 
-UClothingAssetBase* UClothingAssetFactory::Reimport(const FString& Filename, USkeletalMesh* TargetMesh, UClothingAssetBase* OriginalAsset)
+TArray<UClothingAssetBase*> UClothingAssetFactory::CreateFromSourceAsset(USkeletalMesh* TargetMesh, const UObject* SourceAsset) const
 {
-	return nullptr;
+	TArray<UClothingAssetBase*> ClothingAssets;
+	if (const USkeletalMesh* const SourceSkeletalMesh = Cast<USkeletalMesh>(SourceAsset))
+	{
+		for (const UClothingAssetBase* const ClothingAsset : SourceSkeletalMesh->GetMeshClothingAssets())
+		{
+PRAGMA_DISABLE_DEPRECATION_WARNINGS  // UE_DEPRECATED(5.7, ""): CreateFromExistingCloth will be moved out of UClothingAssetFactoryBase to this class
+			if (UClothingAssetBase* const NewAsset =
+				const_cast<UClothingAssetFactory*>(this)->CreateFromExistingCloth(
+					TargetMesh,
+					const_cast<USkeletalMesh*>(SourceSkeletalMesh),
+					const_cast<UClothingAssetBase*>(ClothingAsset)))
+{
+				ClothingAssets.Emplace(NewAsset);
+			}
+PRAGMA_ENABLE_DEPRECATION_WARNINGS
+		}
+	}
+	return ClothingAssets;
 }
 
 UClothingAssetBase* UClothingAssetFactory::CreateFromSkeletalMesh(USkeletalMesh* TargetMesh, FSkeletalMeshClothBuildParams& Params)
@@ -192,9 +204,9 @@ UClothingAssetBase* UClothingAssetFactory::CreateFromSkeletalMesh(USkeletalMesh*
 	return nullptr;
 }
 
-UClothingAssetBase* UClothingAssetFactory::CreateFromExistingCloth(USkeletalMesh* TargetMesh, USkeletalMesh* SourceMesh, UClothingAssetBase* SourceAsset)
+UClothingAssetBase* UClothingAssetFactory::CreateFromExistingCloth(USkeletalMesh* TargetMesh, const USkeletalMesh* /*SourceMesh*/, const UClothingAssetBase* SourceAsset) const
 {
-	UClothingAssetCommon* SourceClothingAsset = Cast<UClothingAssetCommon>(SourceAsset);
+	const UClothingAssetCommon* const SourceClothingAsset = Cast<const UClothingAssetCommon>(SourceAsset);
 
 	if (!SourceClothingAsset)
 	{
@@ -202,7 +214,8 @@ UClothingAssetBase* UClothingAssetFactory::CreateFromExistingCloth(USkeletalMesh
 	}
 
 	//Duplicating the clothing asset using the existing asset as a template
-	UClothingAssetCommon* NewAsset = DuplicateObject<UClothingAssetCommon>(SourceClothingAsset, TargetMesh, SourceClothingAsset->GetFName());
+	const FName NewAssetName = MakeUniqueObjectName(TargetMesh, UClothingAssetCommon::StaticClass(), SourceClothingAsset->GetFName());
+	UClothingAssetCommon* const NewAsset = DuplicateObject<UClothingAssetCommon>(SourceClothingAsset, TargetMesh, NewAssetName);
 
 	NewAsset->AssetGuid = FGuid::NewGuid();
 	//Need to empty LODMap to remove previous mappings from cloth LOD to SkelMesh LOD
@@ -241,8 +254,8 @@ UClothingAssetBase* UClothingAssetFactory::ImportLodToClothing(USkeletalMesh* Ta
 		if(TargetMesh->GetMeshClothingAssets().Find(TargetClothing, MeshAssetIndex))
 		{
 			// Everything looks good, continue to actual import
-			UClothingAssetCommon* ConcreteTarget = CastChecked<UClothingAssetCommon>(TargetClothing);
-
+			if(UClothingAssetCommon* const ConcreteTarget = Cast<UClothingAssetCommon>(TargetClothing))
+			{
 			const FClothLODDataCommon* RemapSource = nullptr;
 
 			if(Params.bRemapParameters)
@@ -292,20 +305,11 @@ UClothingAssetBase* UClothingAssetFactory::ImportLodToClothing(USkeletalMesh* Ta
 			}
 		}
 	}
-
-	return nullptr;
 }
 
-UClothingAssetBase* UClothingAssetFactory::CreateFromApexAsset(nvidia::apex::ClothingAsset* InApexAsset, USkeletalMesh* TargetMesh, FName InName)
-{
 	return nullptr;
 }
 
-bool UClothingAssetFactory::CanImport(const FString& Filename)
-{
-	return false;
-}
-
 bool UClothingAssetFactory::ImportToLodInternal(
 	USkeletalMesh* SourceMesh, 
 	int32 SourceLodIndex, 
diff --git a/Engine/Source/Editor/ClothingSystemEditor/Public/ClothingAssetFactory.h b/Engine/Source/Editor/ClothingSystemEditor/Public/ClothingAssetFactory.h
--- a/Engine/Source/Editor/ClothingSystemEditor/Public/ClothingAssetFactory.h
+++ b/Engine/Source/Editor/ClothingSystemEditor/Public/ClothingAssetFactory.h
@@ -51,24 +51,36 @@ public:
 
 	UE_API UClothingAssetFactory(const FObjectInitializer& ObjectInitializer);
 
-	// Import the given file, treating it as an APEX asset file and return the resulting asset
-	UE_API virtual UClothingAssetBase* Import(const FString& Filename, USkeletalMesh* TargetMesh, FName InName = NAME_None) override;
-	UE_API virtual UClothingAssetBase* Reimport(const FString& Filename, USkeletalMesh* TargetMesh, UClothingAssetBase* OriginalAsset) override;
+	UE_API virtual const UClass* GetSupportedSourceAssetType() const override;
+
+	UE_API virtual TArray<UClothingAssetBase*> CreateFromSourceAsset(USkeletalMesh* TargetMesh, const UObject* SourceAsset) const override;
+
 	UE_API virtual UClothingAssetBase* CreateFromSkeletalMesh(USkeletalMesh* TargetMesh, FSkeletalMeshClothBuildParams& Params) override;
-	UE_API virtual UClothingAssetBase* CreateFromExistingCloth(USkeletalMesh* TargetMesh, USkeletalMesh* SourceMesh, UClothingAssetBase* SourceAsset) override;
+
 	UE_API virtual UClothingAssetBase* ImportLodToClothing(USkeletalMesh* TargetMesh, FSkeletalMeshClothBuildParams& Params) override;
-	// Tests whether the given filename should be able to be imported
-	UE_API virtual bool CanImport(const FString& Filename) override;
 
-	// Given an APEX asset, build a UClothingAssetCommon containing the required data
-	UE_API virtual UClothingAssetBase* CreateFromApexAsset(nvidia::apex::ClothingAsset* InApexAsset, USkeletalMesh* TargetMesh, FName InName = NAME_None) override;
+	// Deprecated methods, cannot be removed because of them being virtual
+	UE_DEPRECATED(5.7, "Use CreateFromSupportedAsset instead.")
+	UE_API virtual UClothingAssetBase* CreateFromExistingCloth(USkeletalMesh* TargetMesh, USkeletalMesh* SourceMesh, UClothingAssetBase* SourceAsset) override
+	{
+		return static_cast<const UClothingAssetFactory*>(this)->CreateFromExistingCloth(TargetMesh, SourceMesh, SourceAsset);
+	}
 
 private:
 
 	// Utility methods for skeletal mesh extraction //////////////////////////
 
+	/**
+	 * Given a target mesh and a source clothing asset bound to source mesh, clone the clothing asset for
+	 * use on the target mesh
+	 * @param TargetMesh - The mesh to target
+	 * @param SourceMesh - The mesh to copy the clothing asset from
+	 * @param SourceAsset - The clothing asset to copy
+	 */
+	UClothingAssetBase* CreateFromExistingCloth(USkeletalMesh* TargetMesh, const USkeletalMesh* SourceMesh, const UClothingAssetBase* SourceAsset) const;
+
 	/** Handles internal import of LODs */
-	UE_API bool ImportToLodInternal(USkeletalMesh* SourceMesh, int32 SourceLodIndex, int32 SourceSectionIndex, UClothingAssetCommon* DestAsset, FClothLODDataCommon& DestLod, int32 DestLodIndex, const FClothLODDataCommon* InParameterRemapSource = nullptr);
+	bool ImportToLodInternal(USkeletalMesh* SourceMesh, int32 SourceLodIndex, int32 SourceSectionIndex, UClothingAssetCommon* DestAsset, FClothLODDataCommon& DestLod, int32 DestLodIndex, const FClothLODDataCommon* InParameterRemapSource = nullptr);
 
 	//////////////////////////////////////////////////////////////////////////
 
diff --git a/Engine/Source/Editor/ClothingSystemEditorInterface/Public/ClothingAssetFactoryInterface.h b/Engine/Source/Editor/ClothingSystemEditorInterface/Public/ClothingAssetFactoryInterface.h
--- a/Engine/Source/Editor/ClothingSystemEditorInterface/Public/ClothingAssetFactoryInterface.h
+++ b/Engine/Source/Editor/ClothingSystemEditorInterface/Public/ClothingAssetFactoryInterface.h
@@ -37,6 +37,10 @@ class UClothingAssetFactoryBase : public UObject
 
 public:
 
+	/** Return the supported source asset type for this clothing data. @see CreateFromSourceAsset. */
+	virtual const UClass* GetSupportedSourceAssetType() const
+	PURE_VIRTUAL(UClothingAssetFactoryBase::GetSupportedSourceAssetType, return nullptr;);
+
 	/**
 	 * Given a target mesh and parameters describing the build operation, create a clothing asset for
 	 * use on the mesh
@@ -46,16 +50,13 @@ public:
 	virtual UClothingAssetBase* CreateFromSkeletalMesh(USkeletalMesh* TargetMesh, FSkeletalMeshClothBuildParams& Params)
 	PURE_VIRTUAL(UClothingAssetFactoryBase::CreateFromSkeletalMesh, return nullptr;);
 
-
 	/**
-	 * Given a target mesh and a source clothing asset bound to source mesh, clone the clothing asset for
-	 * use on the target mesh
+	 * Given a target mesh and a source asset, create clothing data from the source asset to use on the target mesh.
 	 * @param TargetMesh - The mesh to target
-	 * @param SourceMesh - The mesh to copy the clothing asset from
-	 * @param SourceAsset - The clothing asset to copy
+	 * @param SourceAsset - The asset to build the clothing data with.
 	 */
-	virtual UClothingAssetBase* CreateFromExistingCloth(USkeletalMesh* TargetMesh, USkeletalMesh* SourceMesh, UClothingAssetBase* SourceAsset)
-	PURE_VIRTUAL(UClothingAssetFactoryBase::CreateFromExistingCloth, return nullptr;);
+	virtual TArray<UClothingAssetBase*> CreateFromSourceAsset(USkeletalMesh* TargetMesh, const UObject* SourceAsset) const
+	PURE_VIRTUAL(UClothingAssetFactoryBase::CreateFromSourceAsset, return TArray<UClothingAssetBase*>(););
 
 	/** 
 	 * Given a target mesh and valid parameters, import a simulation mesh as a LOD for the clothing
@@ -66,36 +67,36 @@ public:
 	virtual UClothingAssetBase* ImportLodToClothing(USkeletalMesh* TargetMesh, FSkeletalMeshClothBuildParams& Params)
 	PURE_VIRTUAL(UClothingAssetFactoryBase::ImportLodToClothing, return nullptr;);
 
-	/**
-	 * Should return whether or not the factory can handle the incoming file (check validity etc.)
-	 */
-	virtual bool CanImport(const FString& Filename)
-	PURE_VIRTUAL(UClothingAssetFactoryBase::CanImport, return false;);
-
-	/**
-	 * Given an APEX asset instantiated from the filename checked with CanImport, the factory
-	 * is expected to build a valid clothing asset (or nullptr if it cannot).
-	 */
-	virtual UClothingAssetBase* CreateFromApexAsset(nvidia::apex::ClothingAsset* InApexAsset, USkeletalMesh* TargetMesh, FName InName = NAME_None)
-	PURE_VIRTUAL(UClothingAssetFactoryBase::CreateFromApexAsset, return nullptr;);
-
-	/**
-	 * Import an asset from the specified file
-	 * @param Filename - file to import
-	 * @param TargetMesh - the mesh to import the clothing asset to
-	 * @param InName - optional name for the clothing asset object (will take the filename otherwise)
-	 */
+	// Deprecated methods
+	UE_DEPRECATED(5.7, "Clothing file import is no longer supported.")
 	virtual UClothingAssetBase* Import(const FString& Filename, USkeletalMesh* TargetMesh, FName InName = NAME_None)
-	PURE_VIRTUAL(UClothingAssetFactoryBase::Import, return nullptr;);
+	{
+		return nullptr;
+	}
 
-	/**
-	 * Reimport an asset from the specified file
-	 * @param Filename - file to import
-	 * @param TargetMesh - the mesh to import the clothing asset to
-	 * @param OriginalAsset - the original clothing asset to replace
-	 */
+	UE_DEPRECATED(5.7, "Clothing file import is no longer supported.")
 	virtual UClothingAssetBase* Reimport(const FString& Filename, USkeletalMesh* TargetMesh, UClothingAssetBase* OriginalAsset)
-	PURE_VIRTUAL(UClothingAssetFactoryBase::Reimport, return nullptr;);
+	{
+		return nullptr;
+	}
+
+	UE_DEPRECATED(5.7, "Use UClothingAssetFactory::CreateFromSourceAsset() instead.")
+	virtual UClothingAssetBase* CreateFromExistingCloth(USkeletalMesh* /*TargetMesh*/, USkeletalMesh* /*SourceMesh*/, UClothingAssetBase* /*SourceAsset*/)
+	{
+		return nullptr;
+	}
+
+	UE_DEPRECATED(5.7, "Clothing file import is no longer supported.")
+	virtual bool CanImport(const FString& Filename)
+	{
+		return false;
+	}
+
+	UE_DEPRECATED(5.7, "Clothing file import is no longer supported.")
+	virtual UClothingAssetBase* CreateFromApexAsset(nvidia::apex::ClothingAsset* InApexAsset, USkeletalMesh* TargetMesh, FName InName = NAME_None)
+	{
+		return nullptr;
+	}
 };
 
 // An interface for a class that will provide a clothing asset factory, this should be
